Fix NonUnity issue with CUDA Runtime module

diff --git a/Engine/Source/Runtime/CUDA/Source/Private/CudaModule.cpp b/Engine/Source/Runtime/CUDA/Source/Private/CudaModule.cpp
--- a/Engine/Source/Runtime/CUDA/Source/Private/CudaModule.cpp
+++ b/Engine/Source/Runtime/CUDA/Source/Private/CudaModule.cpp
@@ -86,9 +86,12 @@ void FCUDAModule::UnloadCuda()
 
 void FCUDAModule::InitCuda()
 {
+// Don't Init CUDA on unsupported platforms	
+#if PLATFORM_DESKTOP && !PLATFORM_APPLE
+
 	// Initialise to rhiDeviceIndex -1, which is an invalid device index, if it remains -1 we will know no device was found.
 	rhiDeviceIndex = -1;
-	
+
 	// TODO: add support for other RHIs (e.g. DX12)
 	// For now simply exit early if some other RHI is used.
 	if (GDynamicRHI != nullptr && GDynamicRHI->GetName() != FString(TEXT("Vulkan"))) 
@@ -114,11 +117,9 @@ void FCUDAModule::InitCuda()
 	
 	uint8 deviceUUID[16];
 
-#if PLATFORM_DESKTOP && !PLATFORM_APPLE
 	// We find the device that the RHI has selected for us so that later we can create a CUDA context on this device.
 	FVulkanDynamicRHI *vkDynamicRHI = static_cast<FVulkanDynamicRHI *>(GDynamicRHI);
 	FMemory::Memcpy(deviceUUID, vkDynamicRHI->GetDevice()->GetDeviceIdProperties().deviceUUID, 16);	
-#endif
 
 	// Find out how many graphics devices there are so we find that one that matches the one the RHI selected.
 	int device_count = 0;
@@ -202,6 +203,7 @@ void FCUDAModule::InitCuda()
 		UE_LOG(LogCUDA, Fatal, TEXT("CUDA module failed to create a CUDA context on the RHI selected device with UUID %d."), deviceUUID);
 	}
 	
+#endif
 	OnPostCUDAInit.Broadcast();
 }
 
