Fixes an edge case CUDA crash when NVIDIA drivers are present but RHI is running on a non-NVIDIA GPU

diff --git a/Engine/Source/Runtime/CUDA/Source/Private/CudaModule.cpp b/Engine/Source/Runtime/CUDA/Source/Private/CudaModule.cpp
--- a/Engine/Source/Runtime/CUDA/Source/Private/CudaModule.cpp
+++ b/Engine/Source/Runtime/CUDA/Source/Private/CudaModule.cpp
@@ -194,6 +194,17 @@ bool FCUDAModule::IsRHISelectedDevice(CUdevice cuDevice) {
 
 void FCUDAModule::InitCuda()
 {
+	// Unload CUDA if we are not running the RHI on an NVIDIA device this can happen if the NVIDIA drivers
+	// are present on a system with an AMD GPU
+	// TODO is there the potential case that someone will want to use CUDA on a 2nd GPU 
+	// when running the rendering on an AMD one?
+	if (!IsRHIDeviceNVIDIA())
+	{
+		UE_LOG(LogCUDA, Display, TEXT("RHI device not NVIDIA unloading CUDA support"));
+		UnloadCuda();
+		return;
+	}
+
 	// Initialise to rhiDeviceIndex -1, which is an invalid device index, if it remains -1 we will know no device was found.
 	rhiDeviceIndex = -1;
 
