Downgrade CUDA initialization error log from Fatal to Error.

diff --git a/Engine/Source/Runtime/CUDA/Source/Private/CudaModule.cpp b/Engine/Source/Runtime/CUDA/Source/Private/CudaModule.cpp
--- a/Engine/Source/Runtime/CUDA/Source/Private/CudaModule.cpp
+++ b/Engine/Source/Runtime/CUDA/Source/Private/CudaModule.cpp
@@ -221,7 +221,9 @@ void FCUDAModule::InitCuda()
 		}
 		else 
 		{
-			UE_LOG(LogCUDA, Fatal, TEXT("CUDA API failed to initialise."));
+			UE_LOG(LogCUDA, Error, TEXT("CUDA API failed to initialise."));
+			UnloadCuda();
+			return;
 		}
 	}
 
