- Fixing a crash in cuda get context when asking for an invalid device

diff --git a/Engine/Source/Runtime/CUDA/Source/Private/CudaModule.cpp b/Engine/Source/Runtime/CUDA/Source/Private/CudaModule.cpp
--- a/Engine/Source/Runtime/CUDA/Source/Private/CudaModule.cpp
+++ b/Engine/Source/Runtime/CUDA/Source/Private/CudaModule.cpp
@@ -338,7 +338,7 @@ CUcontext FCUDAModule::GetCudaContextForDevice(int DeviceIndex)
 	
 	if(!contextMap.Contains(DeviceIndex))
 	{
-		CUcontext cudaContext;
+		CUcontext cudaContext = nullptr;
 		CUresult createCtxErr = DriverApiPtrs.cuCtxCreate(&cudaContext, 0, DeviceIndex);
 		if(createCtxErr == CUDA_SUCCESS)
 		{
@@ -348,6 +348,7 @@ CUcontext FCUDAModule::GetCudaContextForDevice(int DeviceIndex)
 		else
 		{
 			UE_LOG(LogCUDA, Error, TEXT("Failed to create a CUDA context on device %d on request."), DeviceIndex);
+			return cudaContext;
 		}
 	}
 	
