Skip CUDA initialization for GL RHI on Windows

diff --git a/Engine/Source/Runtime/CUDA/Source/Private/CudaModule.cpp b/Engine/Source/Runtime/CUDA/Source/Private/CudaModule.cpp
--- a/Engine/Source/Runtime/CUDA/Source/Private/CudaModule.cpp
+++ b/Engine/Source/Runtime/CUDA/Source/Private/CudaModule.cpp
@@ -182,6 +182,11 @@ bool FCUDAModule::IsRHISelectedDevice(CUdevice cuDevice)
 					DeviceLuid + sizeof(AdapterLUID.LowPart),
 					sizeof(AdapterLUID.HighPart)) == 0));
 	}
+	else if (RHIType == ERHIInterfaceType::OpenGL)
+	{
+		UE_LOG(LogCUDA, Warning, TEXT("CUDA is unsupported on OpenGL."));
+		return false;
+	}
 	else
 #endif //PLATFORM_WINDOWS
 #endif //PLATFORM_SUPPORTS_CUDA
